Move Cascade editor into its own plugin

#tests ran cascade automation tests with plugin enabled and disabled, checked that asset actions still work, checked that no data corruption occurs when editing an asset with the plugin disabled

diff --git a/Engine/Plugins/FX/Cascade/Source/Cascade/Private/AssetTypeActions_ParticleSystem.cpp b/Engine/Plugins/FX/Cascade/Source/Cascade/Private/AssetTypeActions_ParticleSystem.cpp
new file mode 100644
--- /dev/null
+++ b/Engine/Plugins/FX/Cascade/Source/Cascade/Private/AssetTypeActions_ParticleSystem.cpp
@@ -0,0 +1,132 @@
+// Copyright Epic Games, Inc. All Rights Reserved.
+
+#include "AssetTypeActions_ParticleSystem.h"
+#include "ToolMenus.h"
+#include "Misc/FeedbackContext.h"
+#include "CascadeModule.h"
+#include "Particles/ParticleEmitter.h"
+#include "HAL/PlatformApplicationMisc.h"
+
+#define LOCTEXT_NAMESPACE "AssetTypeActions"
+
+void FAssetTypeActions_ParticleSystem::GetActions(const TArray<UObject*>& InObjects, FToolMenuSection& Section)
+{
+	auto ParticleSystems = GetTypedWeakObjectPtrs<UParticleSystem>(InObjects);
+
+	Section.AddMenuEntry(
+		"ParticleSystem_CopyParameters",
+		LOCTEXT("ParticleSystem_CopyParameters", "Copy Parameters"),
+		LOCTEXT("ParticleSystem_CopyParametersTooltip", "Copies particle system parameters to the clipboard."),
+		FSlateIcon(),
+		FUIAction(
+			FExecuteAction::CreateSP( this, &FAssetTypeActions_ParticleSystem::ExecuteCopyParameters, ParticleSystems ),
+			FCanExecuteAction()
+			)
+		);
+
+	Section.AddMenuEntry(
+		"ParticleSystem_ConvertToSeeded",
+		LOCTEXT("ParticleSystem_ConvertToSeeded", "Convert To Seeded"),
+		LOCTEXT("ParticleSystem_ConvertToSeededTooltip", "Convert all modules in this particle system to random seeded modules"),
+		FSlateIcon(),
+		FUIAction(
+			FExecuteAction::CreateSP( this, &FAssetTypeActions_ParticleSystem::ConvertToSeeded, ParticleSystems ),
+			FCanExecuteAction()
+			)
+		);
+}
+
+void FAssetTypeActions_ParticleSystem::OpenAssetEditor( const TArray<UObject*>& InObjects, TSharedPtr<IToolkitHost> EditWithinLevelEditor )
+{
+	EToolkitMode::Type Mode = EditWithinLevelEditor.IsValid() ? EToolkitMode::WorldCentric : EToolkitMode::Standalone;
+
+	for (auto ObjIt = InObjects.CreateConstIterator(); ObjIt; ++ObjIt)
+	{
+		auto ParticleSystem = Cast<UParticleSystem>(*ObjIt);
+		if (ParticleSystem != NULL)
+		{
+			ICascadeModule* CascadeModule = &FModuleManager::LoadModuleChecked<ICascadeModule>( "Cascade" );
+			CascadeModule->CreateCascade(Mode, EditWithinLevelEditor, ParticleSystem);
+		}
+	}
+}
+
+void FAssetTypeActions_ParticleSystem::ExecuteCopyParameters(TArray<TWeakObjectPtr<UParticleSystem>> Objects)
+{
+	FString ClipboardString = TEXT("");
+	for (auto ObjIt = Objects.CreateConstIterator(); ObjIt; ++ObjIt)
+	{
+		auto Object = (*ObjIt).Get();
+		if ( Object )
+		{
+			TArray<TArray<FString>> ParticleSysParamList;
+			TArray<TArray<FString>> ParticleParameterList;
+			Object->GetParametersUtilized(ParticleSysParamList, ParticleParameterList);
+
+			ClipboardString += FString::Printf(TEXT("ParticleSystem parameters for %s\n"), *(Object->GetPathName()));
+			for (int32 EmitterIndex = 0; EmitterIndex < Object->Emitters.Num(); EmitterIndex++)
+			{
+				ClipboardString += FString::Printf(TEXT("\tEmitter %2d - "), EmitterIndex);
+				UParticleEmitter* Emitter = Object->Emitters[EmitterIndex];
+				if (Emitter)
+				{
+					ClipboardString += FString::Printf(TEXT("%s\n"), *(Emitter->GetEmitterName().ToString()));
+				}
+				else
+				{
+					ClipboardString += FString(TEXT("* EMPTY *\n"));
+				}
+
+				TArray<FString>& ParticleSysParams = ParticleSysParamList[EmitterIndex];
+				for (int32 PSPIndex = 0; PSPIndex < ParticleSysParams.Num(); PSPIndex++)
+				{
+					if (PSPIndex == 0)
+					{
+						ClipboardString += FString(TEXT("\t\tParticleSysParam List\n"));
+					}
+					ClipboardString += FString::Printf(TEXT("\t\t\t%s"), *(ParticleSysParams[PSPIndex]));
+				}
+
+				TArray<FString>& ParticleParameters = ParticleParameterList[EmitterIndex];
+				for (int32 PPIndex = 0; PPIndex < ParticleParameters.Num(); PPIndex++)
+				{
+					if (PPIndex == 0)
+					{
+						ClipboardString += FString(TEXT("\t\tParticleParameter List\n"));
+					}
+					ClipboardString += FString::Printf(TEXT("\t\t\t%s"), *(ParticleParameters[PPIndex]));
+				}
+			}
+		}
+	}
+
+	FPlatformApplicationMisc::ClipboardCopy(*ClipboardString);
+}
+
+void FAssetTypeActions_ParticleSystem::ConvertToSeeded(TArray<TWeakObjectPtr<UParticleSystem>> Objects)
+{
+	ICascadeModule* CascadeModule = &FModuleManager::LoadModuleChecked<ICascadeModule>( "Cascade" );
+
+	if( Objects.Num() > 0 )
+	{
+		GWarn->BeginSlowTask( LOCTEXT("ParticleSystem_ConvertToSeeded_SlowTask", "Converting Particle Systems to Seeded"), true );
+		for (auto ObjIt = Objects.CreateConstIterator(); ObjIt; ++ObjIt)
+		{
+			auto Object = (*ObjIt).Get();
+			if ( Object )
+			{
+				GWarn->StatusUpdate( 
+					ObjIt.GetIndex(), 
+					Objects.Num(),
+					FText::Format( LOCTEXT("ParticleSystem_ConvertToSeeded_StatusUpdate", "Converting {0} to Seeded"), FText::FromString( Object->GetName() ) ) );
+
+				CascadeModule->ConvertModulesToSeeded(Object);
+				CascadeModule->RefreshCascade(Object);
+			}
+		}
+		GWarn->EndSlowTask();
+	}
+
+}
+
+#undef LOCTEXT_NAMESPACE
diff --git a/Engine/Plugins/FX/Cascade/Source/Cascade/Private/AssetTypeActions_ParticleSystem.h b/Engine/Plugins/FX/Cascade/Source/Cascade/Private/AssetTypeActions_ParticleSystem.h
new file mode 100644
--- /dev/null
+++ b/Engine/Plugins/FX/Cascade/Source/Cascade/Private/AssetTypeActions_ParticleSystem.h
@@ -0,0 +1,28 @@
+// Copyright Epic Games, Inc. All Rights Reserved.
+
+#pragma once
+
+#include "CoreMinimal.h"
+#include "Toolkits/IToolkitHost.h"
+#include "AssetTypeActions_Base.h"
+#include "Particles/ParticleSystem.h"
+
+
+class FAssetTypeActions_ParticleSystem : public FAssetTypeActions_Base
+{
+public:
+	// IAssetTypeActions Implementation
+	virtual FText GetName() const override { return NSLOCTEXT("AssetTypeActions", "AssetTypeActions_ParticleSystem", "Cascade Particle System (Deprecated)"); }
+	virtual FColor GetTypeColor() const override { return FColor(255,255,255); }
+	virtual UClass* GetSupportedClass() const override { return UParticleSystem::StaticClass(); }
+	virtual void GetActions(const TArray<UObject*>& InObjects, struct FToolMenuSection& Section) override;
+	virtual void OpenAssetEditor( const TArray<UObject*>& InObjects, TSharedPtr<class IToolkitHost> EditWithinLevelEditor = TSharedPtr<IToolkitHost>() ) override;
+	virtual uint32 GetCategories() override { return EAssetTypeCategories::Misc; }
+	
+private:
+	/** Handler for when Copy Parameters is selected */
+	void ExecuteCopyParameters(TArray<TWeakObjectPtr<UParticleSystem>> Objects);
+
+	/** Handler for when Convert to Seeded is selected */
+	void ConvertToSeeded(TArray<TWeakObjectPtr<UParticleSystem>> Objects);
+};
diff --git a/Engine/Source/Developer/AssetTools/AssetTools.Build.cs b/Engine/Source/Developer/AssetTools/AssetTools.Build.cs
--- a/Engine/Source/Developer/AssetTools/AssetTools.Build.cs
+++ b/Engine/Source/Developer/AssetTools/AssetTools.Build.cs
@@ -64,7 +64,6 @@ public class AssetTools : ModuleRules
 				"PackagesDialog",
 				"Persona",
                 "AnimationEditor",
-				"Cascade",
 				"VirtualTexturingEditor"
 			}
 		);
diff --git a/Engine/Source/Developer/AssetTools/Private/AssetTools.cpp b/Engine/Source/Developer/AssetTools/Private/AssetTools.cpp
--- a/Engine/Source/Developer/AssetTools/Private/AssetTools.cpp
+++ b/Engine/Source/Developer/AssetTools/Private/AssetTools.cpp
@@ -39,7 +39,6 @@
 #include "AssetTypeActions/AssetTypeActions_AnimBoneCompressionSettings.h"
 #include "AssetTypeActions/AssetTypeActions_AnimCurveCompressionSettings.h"
 #include "AssetTypeActions/AssetTypeActions_VariableFrameStrippingSettings.h"
-#include "AssetTypeActions/AssetTypeActions_ParticleSystem.h"
 #include "AssetTypeActions/AssetTypeActions_PhysicalMaterialMask.h"
 #include "AssetTypeActions/AssetTypeActions_DeletedObjectPlaceholder.h"
 #include "WorldPartition/WorldPartition.h"
@@ -1420,7 +1419,6 @@ UAssetToolsImpl::UAssetToolsImpl(const FObjectInitializer& ObjectInitializer)
 	RegisterAssetTypeActions(MakeShareable(new FAssetTypeActions_AnimBoneCompressionSettings));
 	RegisterAssetTypeActions(MakeShareable(new FAssetTypeActions_AnimCurveCompressionSettings));
 	RegisterAssetTypeActions(MakeShareable(new FAssetTypeActions_VariableFrameStrippingSettings));
-	RegisterAssetTypeActions(MakeShareable(new FAssetTypeActions_ParticleSystem));
 	RegisterAssetTypeActions(MakeShareable(new FAssetTypeActions_PhysicalMaterialMask));
 	RegisterAssetTypeActions(MakeShareable(new FAssetTypeActions_DeletedObjectPlaceholder));
 
diff --git a/Engine/Source/Editor/UnrealEd/UnrealEd.Build.cs b/Engine/Source/Editor/UnrealEd/UnrealEd.Build.cs
--- a/Engine/Source/Editor/UnrealEd/UnrealEd.Build.cs
+++ b/Engine/Source/Editor/UnrealEd/UnrealEd.Build.cs
@@ -199,7 +199,6 @@ public class UnrealEd : ModuleRules
 				"FontEditor",
 				"StaticMeshEditor",
 				"TextureEditor",
-				"Cascade",
 				"UMGEditor",
 				"ClassViewer",
 				"StructViewer",
