[ContentBrowser] New Add Menu FX Menu

#Synced-CL 52439714

diff --git a/Engine/Plugins/FX/Cascade/Source/Cascade/Cascade.Build.cs b/Engine/Plugins/FX/Cascade/Source/Cascade/Cascade.Build.cs
--- a/Engine/Plugins/FX/Cascade/Source/Cascade/Cascade.Build.cs
+++ b/Engine/Plugins/FX/Cascade/Source/Cascade/Cascade.Build.cs
@@ -15,7 +15,9 @@ public class Cascade : ModuleRules
 
 		PrivateDependencyModuleNames.AddRange(
 			new string[] {
+				"AssetDefinition",
                 "AppFramework",
+				"ContentBrowser",
 				"Core",
 				"CoreUObject",
                 "InputCore",
diff --git a/Engine/Plugins/FX/Cascade/Source/Cascade/Private/AssetDefinition_ParticleSystem.cpp b/Engine/Plugins/FX/Cascade/Source/Cascade/Private/AssetDefinition_ParticleSystem.cpp
new file mode 100644
--- /dev/null
+++ b/Engine/Plugins/FX/Cascade/Source/Cascade/Private/AssetDefinition_ParticleSystem.cpp
@@ -0,0 +1,172 @@
+// Copyright Epic Games, Inc. All Rights Reserved.
+
+#include "AssetDefinition_ParticleSystem.h"
+
+#include "CascadeModule.h"
+#include "ContentBrowserMenuContexts.h"
+#include "HAL/PlatformApplicationMisc.h"
+#include "Misc/FeedbackContext.h"
+#include "Particles/ParticleEmitter.h"
+#include "ToolMenus.h"
+
+#define LOCTEXT_NAMESPACE "AssetDefinition_ParticleSystem"
+
+EAssetCommandResult UAssetDefinition_ParticleSystem::OpenAssets(const FAssetOpenArgs& OpenArgs) const
+{
+	const EToolkitMode::Type Mode = OpenArgs.ToolkitHost.IsValid() ? EToolkitMode::WorldCentric : EToolkitMode::Standalone;
+
+	const TArray<UObject*>& Objects = OpenArgs.LoadObjects<UObject>();
+	for (auto ObjIt = Objects.CreateConstIterator(); ObjIt; ++ObjIt)
+	{
+		UParticleSystem* ParticleSystem = Cast<UParticleSystem>(*ObjIt);
+		if (ParticleSystem != nullptr)
+		{
+			ICascadeModule* CascadeModule = &FModuleManager::LoadModuleChecked<ICascadeModule>("Cascade");
+			CascadeModule->CreateCascade(Mode, OpenArgs.ToolkitHost, ParticleSystem);
+		}
+	}
+	return EAssetCommandResult::Handled;
+}
+
+namespace MenuExtension::ParticleSystem
+{
+	void ExecuteCopyParameters(const FToolMenuContext& InContext)
+	{
+		const UContentBrowserAssetContextMenuContext* CBContext = UContentBrowserAssetContextMenuContext::FindContextWithAssets(InContext);
+		if (CBContext == nullptr)
+		{
+			return;
+		}
+
+		const TArray<UParticleSystem*> Objects = CBContext->LoadSelectedObjects<UParticleSystem>();
+
+		FString ClipboardString = TEXT("");
+		for (auto ObjIt = Objects.CreateConstIterator(); ObjIt; ++ObjIt)
+		{
+			if (UParticleSystem* Object = *ObjIt)
+			{
+				TArray<TArray<FString>> ParticleSysParamList;
+				TArray<TArray<FString>> ParticleParameterList;
+				Object->GetParametersUtilized(ParticleSysParamList, ParticleParameterList);
+
+				ClipboardString += FString::Printf(TEXT("ParticleSystem parameters for %s\n"), *(Object->GetPathName()));
+				for (int32 EmitterIndex = 0; EmitterIndex < Object->Emitters.Num(); EmitterIndex++)
+				{
+					ClipboardString += FString::Printf(TEXT("\tEmitter %2d - "), EmitterIndex);
+					if (UParticleEmitter* Emitter = Object->Emitters[EmitterIndex])
+					{
+						ClipboardString += FString::Printf(TEXT("%s\n"), *(Emitter->GetEmitterName().ToString()));
+					}
+					else
+					{
+						ClipboardString += FString(TEXT("* EMPTY *\n"));
+					}
+
+					TArray<FString>& ParticleSysParams = ParticleSysParamList[EmitterIndex];
+					for (int32 PSPIndex = 0; PSPIndex < ParticleSysParams.Num(); PSPIndex++)
+					{
+						if (PSPIndex == 0)
+						{
+							ClipboardString += FString(TEXT("\t\tParticleSysParam List\n"));
+						}
+						ClipboardString += FString::Printf(TEXT("\t\t\t%s"), *(ParticleSysParams[PSPIndex]));
+					}
+
+					TArray<FString>& ParticleParameters = ParticleParameterList[EmitterIndex];
+					for (int32 PPIndex = 0; PPIndex < ParticleParameters.Num(); PPIndex++)
+					{
+						if (PPIndex == 0)
+						{
+							ClipboardString += FString(TEXT("\t\tParticleParameter List\n"));
+						}
+						ClipboardString += FString::Printf(TEXT("\t\t\t%s"), *(ParticleParameters[PPIndex]));
+					}
+				}
+			}
+		}
+
+		FPlatformApplicationMisc::ClipboardCopy(*ClipboardString);
+	}
+
+	void ConvertToSeeded(const FToolMenuContext& InContext)
+	{
+		const UContentBrowserAssetContextMenuContext* CBContext = UContentBrowserAssetContextMenuContext::FindContextWithAssets(InContext);
+		if (CBContext == nullptr)
+		{
+			return;
+		}
+
+		const TArray<UParticleSystem*> Objects = CBContext->LoadSelectedObjects<UParticleSystem>();
+
+		ICascadeModule* CascadeModule = &FModuleManager::LoadModuleChecked<ICascadeModule>("Cascade");
+
+		if(Objects.Num() > 0)
+		{
+			GWarn->BeginSlowTask( LOCTEXT("ParticleSystem_ConvertToSeeded_SlowTask", "Converting Particle Systems to Seeded"), true );
+			for (auto ObjIt = Objects.CreateConstIterator(); ObjIt; ++ObjIt)
+			{
+				if (UParticleSystem* Object = *ObjIt)
+				{
+					GWarn->StatusUpdate( 
+						ObjIt.GetIndex(), 
+						Objects.Num(),
+						FText::Format( LOCTEXT("ParticleSystem_ConvertToSeeded_StatusUpdate", "Converting {0} to Seeded"), FText::FromString( Object->GetName() ) ) );
+
+					CascadeModule->ConvertModulesToSeeded(Object);
+					CascadeModule->RefreshCascade(Object);
+				}
+			}
+			GWarn->EndSlowTask();
+		}
+	}
+		static FDelayedAutoRegisterHelper DelayedAutoRegister(EDelayedRegisterRunPhase::EndOfEngineInit, []
+	{
+		UToolMenus::RegisterStartupCallback(FSimpleMulticastDelegate::FDelegate::CreateLambda([]()
+		{
+			FToolMenuOwnerScoped OwnerScoped(UE_MODULE_NAME);
+			
+			// Particle System Action Registration
+			{
+				UToolMenu* Menu = UE::ContentBrowser::ExtendToolMenu_AssetContextMenu(UParticleSystem::StaticClass());
+
+				FToolMenuSection& ParticleSystemSection = Menu->FindOrAddSection("GetAssetActions");
+				ParticleSystemSection.AddDynamicEntry(NAME_None, FNewToolMenuSectionDelegate::CreateLambda([](FToolMenuSection& InSection)
+					{
+						if (const UContentBrowserAssetContextMenuContext* Context = UContentBrowserAssetContextMenuContext::FindContextWithAssets(InSection))
+						{
+							if (Context->SelectedAssets.Num() > 0)
+							{
+								const FAssetData& AssetData = Context->SelectedAssets[0];
+								if (AssetData.AssetClassPath == UParticleSystem::StaticClass()->GetClassPathName())
+								{
+									// Copy
+									{
+										const TAttribute<FText> Label = LOCTEXT("ParticleSystem_CopyParameters", "Copy Parameters");
+										const TAttribute<FText> ToolTip = LOCTEXT("ParticleSystem_CopyParametersTooltip", "Copies particle system parameters to the clipboard.");
+										const FSlateIcon Icon = FSlateIcon(FAppStyle::GetAppStyleSetName(), "GenericCommands.Copy");
+
+										FToolUIAction UIAction;
+										UIAction.ExecuteAction = FToolMenuExecuteAction::CreateStatic(&ExecuteCopyParameters);
+										InSection.AddMenuEntry("ParticleSystem_CopyParameters", Label, ToolTip, Icon, UIAction);
+									}
+
+									// Convert
+									{
+										const TAttribute<FText> Label = LOCTEXT("ParticleSystem_ConvertToSeeded", "Convert To Seeded");
+										const TAttribute<FText> ToolTip = LOCTEXT("ParticleSystem_ConvertToSeeded", "Convert all modules in this particle system to random seeded modules.");
+										const FSlateIcon Icon = FSlateIcon(FAppStyle::GetAppStyleSetName(), "Icons.Refresh");
+
+										FToolUIAction UIAction;
+										UIAction.ExecuteAction = FToolMenuExecuteAction::CreateStatic(&ConvertToSeeded);
+										InSection.AddMenuEntry("ParticleSystem_ConvertToSeeded", Label, ToolTip, Icon, UIAction);
+									}
+								}
+							}
+						}
+					}));
+			}
+		}));
+	});
+}
+
+#undef LOCTEXT_NAMESPACE
diff --git a/Engine/Plugins/FX/Cascade/Source/Cascade/Private/AssetDefinition_ParticleSystem.h b/Engine/Plugins/FX/Cascade/Source/Cascade/Private/AssetDefinition_ParticleSystem.h
new file mode 100644
--- /dev/null
+++ b/Engine/Plugins/FX/Cascade/Source/Cascade/Private/AssetDefinition_ParticleSystem.h
@@ -0,0 +1,24 @@
+// Copyright Epic Games, Inc. All Rights Reserved.
+
+#pragma once
+#include "AssetDefinitionDefault.h"
+#include "Particles/ParticleSystem.h"
+#include "AssetDefinition_ParticleSystem.generated.h"
+
+UCLASS()
+class UAssetDefinition_ParticleSystem : public UAssetDefinitionDefault
+{
+	GENERATED_BODY()
+public:
+	// UAssetDefinition Begin
+	virtual FText GetAssetDisplayName() const override { return NSLOCTEXT("AssetDefinition", "AssetDefinition_ParticleSystem", "Cascade Particle System (Deprecated)"); }
+	virtual FLinearColor GetAssetColor() const override { return FColor(255,255,255); }
+	virtual TSoftClassPtr<UObject> GetAssetClass() const override { return UParticleSystem::StaticClass(); }
+	virtual TConstArrayView<FAssetCategoryPath> GetAssetCategories() const override
+	{
+		static const auto Categories = { FAssetCategoryPath(EAssetCategoryPaths::FX, NSLOCTEXT("AssetDefinition", "AssetDefinition_ParticleSystemSubMenu", "Deprecated"), ECategoryMenuType::Section) };
+		return Categories;
+	}
+	virtual EAssetCommandResult OpenAssets(const FAssetOpenArgs& OpenArgs) const override;
+	// UAssetDefinition End
+};
diff --git a/Engine/Plugins/FX/Cascade/Source/Cascade/Private/AssetTypeActions_ParticleSystem.cpp b/Engine/Plugins/FX/Cascade/Source/Cascade/Private/AssetTypeActions_ParticleSystem.cpp
deleted file mode 100644
diff --git a/Engine/Plugins/FX/Cascade/Source/Cascade/Private/AssetTypeActions_ParticleSystem.h b/Engine/Plugins/FX/Cascade/Source/Cascade/Private/AssetTypeActions_ParticleSystem.h
deleted file mode 100644
diff --git a/Engine/Plugins/FX/Cascade/Source/Cascade/Private/CascadeModule.cpp b/Engine/Plugins/FX/Cascade/Source/Cascade/Private/CascadeModule.cpp
--- a/Engine/Plugins/FX/Cascade/Source/Cascade/Private/CascadeModule.cpp
+++ b/Engine/Plugins/FX/Cascade/Source/Cascade/Private/CascadeModule.cpp
@@ -2,7 +2,6 @@
 
 #include "CascadeModule.h"
 
-#include "AssetTypeActions_ParticleSystem.h"
 #include "Cascade.h"
 #include "Containers/Array.h"
 #include "Delegates/Delegate.h"
@@ -36,16 +35,6 @@ public:
 		ToolBarExtensibilityManager = MakeShareable(new FExtensibilityManager);
 
 		UParticleSystemComponent::OnSystemPreActivationChange.AddStatic(&FCascade::OnComponentActivationChange);
-		
-		// Register asset tools
-		auto RegisterAssetTypeAction = [this](const TSharedRef<IAssetTypeActions>& InAssetTypeAction)
-		{
-			IAssetTools& AssetTools = FModuleManager::LoadModuleChecked<FAssetToolsModule>("AssetTools").Get();
-			RegisteredAssetTypeActions.Add(InAssetTypeAction);
-			AssetTools.RegisterAssetTypeActions(InAssetTypeAction);
-		};
-
-		RegisterAssetTypeAction(MakeShareable(new FAssetTypeActions_ParticleSystem()));
 	}
 
 	/** Called before the module is unloaded, right before the module object is destroyed. */
@@ -53,15 +42,6 @@ public:
 	{
 		MenuExtensibilityManager.Reset();
 		ToolBarExtensibilityManager.Reset();
-
-		if (FAssetToolsModule* AssetToolsModule = FModuleManager::GetModulePtr<FAssetToolsModule>("AssetTools"))
-		{
-			for (TSharedRef<IAssetTypeActions> RegisteredAssetTypeAction : RegisteredAssetTypeActions)
-			{
-				AssetToolsModule->Get().UnregisterAssetTypeActions(RegisteredAssetTypeAction);
-			}
-			RegisteredAssetTypeActions.Reset();
-		}
 	}
 
 	virtual TSharedRef<ICascade> CreateCascade(const EToolkitMode::Type Mode, const TSharedPtr< IToolkitHost >& InitToolkitHost, UParticleSystem* ParticleSystem) override
diff --git a/Engine/Plugins/FX/Cascade/Source/Cascade/Public/CascadeModule.h b/Engine/Plugins/FX/Cascade/Source/Cascade/Public/CascadeModule.h
--- a/Engine/Plugins/FX/Cascade/Source/Cascade/Public/CascadeModule.h
+++ b/Engine/Plugins/FX/Cascade/Source/Cascade/Public/CascadeModule.h
@@ -36,7 +36,4 @@ public:
 
 	/** Converts all the modules in the specified particle system to seeded modules */
 	virtual void ConvertModulesToSeeded(UParticleSystem* ParticleSystem) = 0;
-	
-protected:
-	TArray<TSharedRef<IAssetTypeActions>> RegisteredAssetTypeActions;
 };
diff --git a/Engine/Plugins/FX/Niagara/Source/NiagaraEditor/Private/AssetDefinitions/AssetDefinition_NiagaraDataChannel.cpp b/Engine/Plugins/FX/Niagara/Source/NiagaraEditor/Private/AssetDefinitions/AssetDefinition_NiagaraDataChannel.cpp
--- a/Engine/Plugins/FX/Niagara/Source/NiagaraEditor/Private/AssetDefinitions/AssetDefinition_NiagaraDataChannel.cpp
+++ b/Engine/Plugins/FX/Niagara/Source/NiagaraEditor/Private/AssetDefinitions/AssetDefinition_NiagaraDataChannel.cpp
@@ -24,7 +24,7 @@ TSoftClassPtr<> UAssetDefinition_NiagaraDataChannel::GetAssetClass() const
 
 TConstArrayView<FAssetCategoryPath> UAssetDefinition_NiagaraDataChannel::GetAssetCategories() const
 {
-	static FAssetCategoryPath AssetPaths[] = { FAssetCategoryPath(LOCTEXT("NiagaraAssetsCategory", "FX"), LOCTEXT("NiagaraDataChannel_SubCategory", "Advanced")) };
+	static FAssetCategoryPath AssetPaths[] = { FAssetCategoryPath(EAssetCategoryPaths::FX, LOCTEXT("NiagaraDataChannel_SubCategory", "Advanced"), ECategoryMenuType::Section) };
 	return AssetPaths;
 }
 
diff --git a/Engine/Plugins/FX/Niagara/Source/NiagaraEditor/Private/AssetDefinitions/AssetDefinition_NiagaraEffectType.h b/Engine/Plugins/FX/Niagara/Source/NiagaraEditor/Private/AssetDefinitions/AssetDefinition_NiagaraEffectType.h
--- a/Engine/Plugins/FX/Niagara/Source/NiagaraEditor/Private/AssetDefinitions/AssetDefinition_NiagaraEffectType.h
+++ b/Engine/Plugins/FX/Niagara/Source/NiagaraEditor/Private/AssetDefinitions/AssetDefinition_NiagaraEffectType.h
@@ -21,7 +21,7 @@ public:
 	virtual TConstArrayView<FAssetCategoryPath> GetAssetCategories() const override
 	{
 		static const auto Categories = {
-			EAssetCategoryPaths::FX / NSLOCTEXT("Niagara", "NiagaraAssetSubMenu_Advanced", "Advanced")
+			FAssetCategoryPath(EAssetCategoryPaths::FX, NSLOCTEXT("Niagara", "NiagaraAssetSubMenu_Advanced", "Advanced"), ECategoryMenuType::Section)
 		};
 		
 		return Categories;
diff --git a/Engine/Plugins/FX/Niagara/Source/NiagaraEditor/Private/AssetDefinitions/AssetDefinition_NiagaraEmitter.h b/Engine/Plugins/FX/Niagara/Source/NiagaraEditor/Private/AssetDefinitions/AssetDefinition_NiagaraEmitter.h
--- a/Engine/Plugins/FX/Niagara/Source/NiagaraEditor/Private/AssetDefinitions/AssetDefinition_NiagaraEmitter.h
+++ b/Engine/Plugins/FX/Niagara/Source/NiagaraEditor/Private/AssetDefinitions/AssetDefinition_NiagaraEmitter.h
@@ -20,7 +20,7 @@ public:
 	virtual TSoftClassPtr<UObject> GetAssetClass() const override { return UNiagaraEmitter::StaticClass(); }
 	virtual TConstArrayView<FAssetCategoryPath> GetAssetCategories() const override
 	{
-		static const auto Categories = { EAssetCategoryPaths::FX };
+		static const auto Categories = { FAssetCategoryPath(EAssetCategoryPaths::FX) };
 		return Categories;
 	}
 
diff --git a/Engine/Plugins/FX/Niagara/Source/NiagaraEditor/Private/AssetDefinitions/AssetDefinition_NiagaraParameterCollection.h b/Engine/Plugins/FX/Niagara/Source/NiagaraEditor/Private/AssetDefinitions/AssetDefinition_NiagaraParameterCollection.h
--- a/Engine/Plugins/FX/Niagara/Source/NiagaraEditor/Private/AssetDefinitions/AssetDefinition_NiagaraParameterCollection.h
+++ b/Engine/Plugins/FX/Niagara/Source/NiagaraEditor/Private/AssetDefinitions/AssetDefinition_NiagaraParameterCollection.h
@@ -20,7 +20,7 @@ public:
 	virtual TSoftClassPtr<UObject> GetAssetClass() const override { return UNiagaraParameterCollection::StaticClass(); }
 	virtual TConstArrayView<FAssetCategoryPath> GetAssetCategories() const override
 	{
-		static const auto Categories = { EAssetCategoryPaths::FX / NSLOCTEXT("Niagara", "NiagaraAssetSubMenu_Advanced", "Advanced") };
+		static const auto Categories = { FAssetCategoryPath(EAssetCategoryPaths::FX, NSLOCTEXT("Niagara", "NiagaraAssetSubMenu_Advanced", "Advanced"), ECategoryMenuType::Section) };
 		return Categories;
 	}
 
diff --git a/Engine/Plugins/FX/Niagara/Source/NiagaraEditor/Private/AssetDefinitions/AssetDefinition_NiagaraParameterCollectionInstance.h b/Engine/Plugins/FX/Niagara/Source/NiagaraEditor/Private/AssetDefinitions/AssetDefinition_NiagaraParameterCollectionInstance.h
--- a/Engine/Plugins/FX/Niagara/Source/NiagaraEditor/Private/AssetDefinitions/AssetDefinition_NiagaraParameterCollectionInstance.h
+++ b/Engine/Plugins/FX/Niagara/Source/NiagaraEditor/Private/AssetDefinitions/AssetDefinition_NiagaraParameterCollectionInstance.h
@@ -20,7 +20,7 @@ public:
 	virtual TSoftClassPtr<UObject> GetAssetClass() const override { return UNiagaraParameterCollectionInstance::StaticClass(); }
 	virtual TConstArrayView<FAssetCategoryPath> GetAssetCategories() const override
 	{
-		static const auto Categories = { EAssetCategoryPaths::FX / NSLOCTEXT("Niagara", "NiagaraAssetSubMenu_Advanced", "Advanced") };
+		static const auto Categories = { FAssetCategoryPath(EAssetCategoryPaths::FX, NSLOCTEXT("Niagara", "NiagaraAssetSubMenu_Advanced", "Advanced"), ECategoryMenuType::Section) };
 		return Categories;
 	}
 
diff --git a/Engine/Plugins/FX/Niagara/Source/NiagaraEditor/Private/AssetDefinitions/AssetDefinition_NiagaraParameterDefinitions.h b/Engine/Plugins/FX/Niagara/Source/NiagaraEditor/Private/AssetDefinitions/AssetDefinition_NiagaraParameterDefinitions.h
--- a/Engine/Plugins/FX/Niagara/Source/NiagaraEditor/Private/AssetDefinitions/AssetDefinition_NiagaraParameterDefinitions.h
+++ b/Engine/Plugins/FX/Niagara/Source/NiagaraEditor/Private/AssetDefinitions/AssetDefinition_NiagaraParameterDefinitions.h
@@ -20,7 +20,7 @@ public:
 	virtual TSoftClassPtr<UObject> GetAssetClass() const override { return UNiagaraParameterDefinitions::StaticClass(); }
 	virtual TConstArrayView<FAssetCategoryPath> GetAssetCategories() const override
 	{
-		static const auto Categories = { EAssetCategoryPaths::FX / NSLOCTEXT("Niagara", "NiagaraAssetSubMenu_Advanced", "Advanced") };
+		static const auto Categories = { FAssetCategoryPath(EAssetCategoryPaths::FX, NSLOCTEXT("Niagara", "NiagaraAssetSubMenu_Advanced", "Advanced"), ECategoryMenuType::Section) };
 		return Categories;
 	}
 
diff --git a/Engine/Plugins/FX/Niagara/Source/NiagaraEditor/Private/AssetDefinitions/AssetDefinition_NiagaraScript.h b/Engine/Plugins/FX/Niagara/Source/NiagaraEditor/Private/AssetDefinitions/AssetDefinition_NiagaraScript.h
--- a/Engine/Plugins/FX/Niagara/Source/NiagaraEditor/Private/AssetDefinitions/AssetDefinition_NiagaraScript.h
+++ b/Engine/Plugins/FX/Niagara/Source/NiagaraEditor/Private/AssetDefinitions/AssetDefinition_NiagaraScript.h
@@ -31,7 +31,7 @@ public:
 	virtual TSoftClassPtr<UObject> GetAssetClass() const override { return UNiagaraScript::StaticClass(); }
 	virtual TConstArrayView<FAssetCategoryPath> GetAssetCategories() const override
 	{
-		static const auto Categories = { EAssetCategoryPaths::FX / NSLOCTEXT("Niagara", "NiagaraAssetSubMenu_Script", "Script") };
+		static const auto Categories = { FAssetCategoryPath(EAssetCategoryPaths::FX, NSLOCTEXT("Niagara", "NiagaraAssetSubMenu_Script", "Script"), ECategoryMenuType::Section) };
 		return Categories;
 	}
 
diff --git a/Engine/Plugins/FX/Niagara/Source/NiagaraEditor/Private/AssetDefinitions/AssetDefinition_NiagaraSimCache.h b/Engine/Plugins/FX/Niagara/Source/NiagaraEditor/Private/AssetDefinitions/AssetDefinition_NiagaraSimCache.h
--- a/Engine/Plugins/FX/Niagara/Source/NiagaraEditor/Private/AssetDefinitions/AssetDefinition_NiagaraSimCache.h
+++ b/Engine/Plugins/FX/Niagara/Source/NiagaraEditor/Private/AssetDefinitions/AssetDefinition_NiagaraSimCache.h
@@ -20,7 +20,7 @@ public:
 	virtual TSoftClassPtr<UObject> GetAssetClass() const override { return UNiagaraSimCache::StaticClass(); }
 	virtual TConstArrayView<FAssetCategoryPath> GetAssetCategories() const override
 	{
-		static const auto Categories = { EAssetCategoryPaths::FX / NSLOCTEXT("Niagara", "NiagaraAssetSubMenu_Advanced", "Advanced") };
+		static const auto Categories = { FAssetCategoryPath(EAssetCategoryPaths::FX, NSLOCTEXT("Niagara", "NiagaraAssetSubMenu_Advanced", "Advanced"), ECategoryMenuType::Section) };
 		return Categories;
 	}
 
diff --git a/Engine/Plugins/FX/Niagara/Source/NiagaraEditor/Private/AssetDefinitions/AssetDefinition_NiagaraValidationRuleSet.h b/Engine/Plugins/FX/Niagara/Source/NiagaraEditor/Private/AssetDefinitions/AssetDefinition_NiagaraValidationRuleSet.h
--- a/Engine/Plugins/FX/Niagara/Source/NiagaraEditor/Private/AssetDefinitions/AssetDefinition_NiagaraValidationRuleSet.h
+++ b/Engine/Plugins/FX/Niagara/Source/NiagaraEditor/Private/AssetDefinitions/AssetDefinition_NiagaraValidationRuleSet.h
@@ -21,7 +21,7 @@ public:
 	virtual TConstArrayView<FAssetCategoryPath> GetAssetCategories() const override
 	{
 		static const auto Categories = {
-			EAssetCategoryPaths::FX / NSLOCTEXT("Niagara", "NiagaraAssetSubMenu_Advanced", "Advanced")
+			FAssetCategoryPath(EAssetCategoryPaths::FX, NSLOCTEXT("Niagara", "NiagaraAssetSubMenu_Advanced", "Advanced"), ECategoryMenuType::Section)
 		};
 		
 		return Categories;
