[Unreal Editor Localization] Fix localization issue

#Synced-CL 52533865

diff --git a/Engine/Plugins/FX/Cascade/Source/Cascade/Private/AssetDefinition_ParticleSystem.cpp b/Engine/Plugins/FX/Cascade/Source/Cascade/Private/AssetDefinition_ParticleSystem.cpp
--- a/Engine/Plugins/FX/Cascade/Source/Cascade/Private/AssetDefinition_ParticleSystem.cpp
+++ b/Engine/Plugins/FX/Cascade/Source/Cascade/Private/AssetDefinition_ParticleSystem.cpp
@@ -153,7 +153,7 @@ namespace MenuExtension::ParticleSystem
 									// Convert
 									{
 										const TAttribute<FText> Label = LOCTEXT("ParticleSystem_ConvertToSeeded", "Convert To Seeded");
-										const TAttribute<FText> ToolTip = LOCTEXT("ParticleSystem_ConvertToSeeded", "Convert all modules in this particle system to random seeded modules.");
+										const TAttribute<FText> ToolTip = LOCTEXT("ParticleSystem_ConvertToSeededToolTip", "Convert all modules in this particle system to random seeded modules.");
 										const FSlateIcon Icon = FSlateIcon(FAppStyle::GetAppStyleSetName(), "Icons.Refresh");
 
 										FToolUIAction UIAction;
